Rename ChaosUserData plugin to ChaosUserDataPT

diff --git a/Engine/Plugins/Experimental/ChaosUserData/ChaosUserData.uplugin b/Engine/Plugins/Experimental/ChaosUserData/ChaosUserData.uplugin
deleted file mode 100644
diff --git a/Engine/Plugins/Experimental/ChaosUserData/Source/ChaosUserData/ChaosUserData.Build.cs b/Engine/Plugins/Experimental/ChaosUserData/Source/ChaosUserData/ChaosUserData.Build.cs
deleted file mode 100644
diff --git a/Engine/Plugins/Experimental/ChaosUserData/Source/ChaosUserData/Private/ChaosUserDataModule.cpp b/Engine/Plugins/Experimental/ChaosUserData/Source/ChaosUserData/Private/ChaosUserDataModule.cpp
deleted file mode 100644
diff --git a/Engine/Plugins/Experimental/ChaosUserData/Source/ChaosUserData/Private/ChaosUserDataStats.cpp b/Engine/Plugins/Experimental/ChaosUserData/Source/ChaosUserData/Private/ChaosUserDataStats.cpp
deleted file mode 100644
diff --git a/Engine/Plugins/Experimental/ChaosUserData/Source/ChaosUserData/Public/ChaosUserData.h b/Engine/Plugins/Experimental/ChaosUserData/Source/ChaosUserData/Public/ChaosUserData.h
deleted file mode 100644
diff --git a/Engine/Plugins/Experimental/ChaosUserData/Source/ChaosUserData/Public/ChaosUserDataModule.h b/Engine/Plugins/Experimental/ChaosUserData/Source/ChaosUserData/Public/ChaosUserDataModule.h
deleted file mode 100644
diff --git a/Engine/Plugins/Experimental/ChaosUserData/Source/ChaosUserData/Public/ChaosUserDataStats.h b/Engine/Plugins/Experimental/ChaosUserData/Source/ChaosUserData/Public/ChaosUserDataStats.h
deleted file mode 100644
diff --git a/Engine/Plugins/Experimental/ChaosUserDataPT/ChaosUserDataPT.uplugin b/Engine/Plugins/Experimental/ChaosUserDataPT/ChaosUserDataPT.uplugin
new file mode 100644
--- /dev/null
+++ b/Engine/Plugins/Experimental/ChaosUserDataPT/ChaosUserDataPT.uplugin
@@ -0,0 +1,24 @@
+{
+	"FileVersion": 1,
+	"Version": 1,
+	"VersionName": "1.0",
+	"FriendlyName": "ChaosUserDataPT",
+	"Description": "Custom per-particle userdata. Write-only on game thread, read-only on physics thread.",
+	"Category": "Editor",
+	"CreatedBy": "Epic Games, Inc.",
+	"CreatedByURL": "",
+	"DocsURL": "",
+	"MarketplaceURL": "",
+	"SupportURL": "",
+	"CanContainContent": false,
+	"IsBetaVersion": true,
+	"Installed": false,
+	"EnabledByDefault": true,
+	"Modules": [
+		{
+			"Name": "ChaosUserDataPT",
+			"Type": "Runtime",
+			"LoadingPhase": "Default"
+		}
+	]
+}
\ No newline at end of file
diff --git a/Engine/Plugins/Experimental/ChaosUserDataPT/Source/ChaosUserData/ChaosUserDataPT.Build.cs b/Engine/Plugins/Experimental/ChaosUserDataPT/Source/ChaosUserData/ChaosUserDataPT.Build.cs
new file mode 100644
--- /dev/null
+++ b/Engine/Plugins/Experimental/ChaosUserDataPT/Source/ChaosUserData/ChaosUserDataPT.Build.cs
@@ -0,0 +1,23 @@
+// Copyright Epic Games, Inc. All Rights Reserved.
+
+using UnrealBuildTool;
+
+namespace UnrealBuildTool.Rules
+{
+	public class ChaosUserDataPT : ModuleRules
+	{
+
+		public ChaosUserDataPT(ReadOnlyTargetRules Target) : base(Target)
+		{
+			PublicDependencyModuleNames.AddRange(
+				new string[]
+				{
+					"Core",
+				}
+			);
+
+			SetupModulePhysicsSupport(Target);
+			PrivateDefinitions.Add("CHAOS_INCLUDE_LEVEL_1=1");
+		}
+	}
+}
diff --git a/Engine/Plugins/Experimental/ChaosUserDataPT/Source/ChaosUserData/Private/ChaosUserDataPTModule.cpp b/Engine/Plugins/Experimental/ChaosUserDataPT/Source/ChaosUserData/Private/ChaosUserDataPTModule.cpp
new file mode 100644
--- /dev/null
+++ b/Engine/Plugins/Experimental/ChaosUserDataPT/Source/ChaosUserData/Private/ChaosUserDataPTModule.cpp
@@ -0,0 +1,16 @@
+// Copyright Epic Games, Inc. All Rights Reserved.
+
+#include "ChaosUserDataPTModule.h"
+
+void FChaosUserDataPTModule::StartupModule()
+{
+	// This code will execute after your module is loaded into memory; the exact timing is specified in the .uplugin file per-module
+}
+
+void FChaosUserDataPTModule::ShutdownModule()
+{
+	// This function may be called during shutdown to clean up your module.  For modules that support dynamic reloading,
+	// we call this function before unloading the module.
+}
+
+IMPLEMENT_MODULE(FChaosUserDataPTModule, ChaosUserDataPT)
\ No newline at end of file
diff --git a/Engine/Plugins/Experimental/ChaosUserDataPT/Source/ChaosUserData/Private/ChaosUserDataPTStats.cpp b/Engine/Plugins/Experimental/ChaosUserDataPT/Source/ChaosUserData/Private/ChaosUserDataPTStats.cpp
new file mode 100644
--- /dev/null
+++ b/Engine/Plugins/Experimental/ChaosUserDataPT/Source/ChaosUserData/Private/ChaosUserDataPTStats.cpp
@@ -0,0 +1,12 @@
+// Copyright Epic Games, Inc. All Rights Reserved.
+
+#include "ChaosUserDataPTStats.h"
+
+namespace Chaos
+{
+	DEFINE_STAT(STAT_UserDataPT_SetData_GT);
+	DEFINE_STAT(STAT_UserDataPT_RemoveData_GT);
+	DEFINE_STAT(STAT_UserDataPT_GetData_PT);
+	DEFINE_STAT(STAT_UserDataPT_UpdateData_PT);
+	DEFINE_STAT(STAT_UserDataPT_RemoveData_PT);
+}
diff --git a/Engine/Plugins/Experimental/ChaosUserDataPT/Source/ChaosUserData/Public/ChaosUserDataPT.h b/Engine/Plugins/Experimental/ChaosUserDataPT/Source/ChaosUserData/Public/ChaosUserDataPT.h
new file mode 100644
--- /dev/null
+++ b/Engine/Plugins/Experimental/ChaosUserDataPT/Source/ChaosUserData/Public/ChaosUserDataPT.h
@@ -0,0 +1,216 @@
+// Copyright Epic Games, Inc. All Rights Reserved.
+
+#pragma once
+
+#include "CoreMinimal.h"
+#include "Chaos/PBDRigidsEvolutionFwd.h"
+#include "Chaos/SimCallbackInput.h"
+#include "Chaos/SimCallbackObject.h"
+#include "Chaos/GeometryParticlesfwd.h"
+#include "Chaos/Framework/PhysicsSolverBase.h"
+#include "UObject/WeakObjectPtr.h"
+#include "PhysicsInterfaceDeclaresCore.h"
+#include "PhysicsProxy/SingleParticlePhysicsProxy.h"
+#include "Templates/UniquePtr.h"
+#include "ChaosUserDataPTStats.h"
+
+/*
+
+	Chaos User Data PT
+	==================
+
+	The idea behind this tool is to provide a generic way of associating
+	custom data with physics particles, which is write-only from the game
+	thread, and read-only from the physics thread.
+
+	This comes in handy when physical interactions at the per-contact level
+	need to be affected by gameplay properties.
+
+	A TUserDataManagerPT is a sim callback which is templated on a type TUserData. It
+	owns an instance of TUserData for each particle which has been used as an
+	argument to SetData_GT.
+
+	In order to use a TUserDataManagerPT, it will need to be created using the chaos
+	solver's CreateAndRegisterSimCallbackObject_External method. This
+	library does not natively provide a method of accessing the
+	appropriate TUserDataManagerPT instance from the physics thread, but this can
+	be achieved in a number of ways - it is left up to the game to decide
+	how to do this for flexibility.
+
+*/
+
+namespace Chaos
+{
+	// TUserDataManagerPTInput
+	//
+	// Input is a collection of pointers to new and updated userdata objects
+	// to be sent to the physics thread
+	template <typename TUserData>
+	struct TUserDataManagerPTInput : public FSimCallbackInput
+	{
+		// Map of particle unique indices to user data ptrs
+		// 
+		// NOTE: This is marked mutable because the userdata objects must be
+		// moved to the internal array after making it
+		// OnPreSimulate_Internal, but input objects are const in that
+		// context. This might be frowned on, but since TUserDataManagerPTInput is a
+		// class which is only used internally to TUserDataManagerPT, an argument
+		// could be made either way.
+		mutable TMap<FUniqueIdx, TUniquePtr<TUserData>> UserDataToAdd;
+
+		// Set of particle unique indices for which to remove user data
+		TSet<FUniqueIdx> UserDataToRemove;
+
+		// Monotonically increasing identifier for the input object. Each
+		// newly constructed input will store and increment this;
+		int32 Identifier = -1;
+
+		void Reset()
+		{
+			UserDataToAdd.Reset();
+			UserDataToRemove.Reset();
+			Identifier = -1;
+		}
+	};
+
+	// TUserDataManagerPT
+	//
+	// A chaos callback object which stores and allows access to user data
+	// associated with particles on the physics thread.
+	//
+	// Note that FSimCallbackOutput is the output struct - this carries no
+	// data because this is a one-way callback. We use it basically just to
+	// marshal data in one direction.
+	template <typename TUserData>
+	class TUserDataManagerPT : public TSimCallbackObject<TUserDataManagerPTInput<TUserData>, FSimCallbackNoOutput>
+	{
+		using TInput = TUserDataManagerPTInput<TUserData>;
+
+	public:
+
+		virtual ~TUserDataManagerPT() { }
+
+		// Add or update user data associated with this particle handle
+		bool SetData_GT(const FRigidBodyHandle_External& Handle, const TUserData& UserData)
+		{
+			SCOPE_CYCLE_COUNTER(STAT_UserDataPT_SetData_GT);
+
+			if (const FPhysicsSolverBase* Solver = this->GetSolver())
+			{
+				if (TInput* Input = this->GetProducerInputData_External())
+				{
+					// Add the data to the map to be sent to physics thread
+					Input->UserDataToAdd.Emplace(Handle.UniqueIdx(), MakeUnique<TUserData>(UserData));
+
+					// In case it was removed and then added again in the same
+					// frame, untrack this particle for data removal
+					Input->UserDataToRemove.Remove(Handle.UniqueIdx());
+
+					// If this is a new input, set it's identifier
+					if (Input->Identifier == -1)
+					{
+						Input->Identifier = InputIdentifier_GT++;
+					}
+
+					// Successfully queued for add/update
+					return true;
+				}
+			}
+
+			// Failed to queue for add/update
+			return false;
+		}
+
+		// Remove user data associated with this particle handle
+		bool RemoveData_GT(const FRigidBodyHandle_External& Handle)
+		{
+			SCOPE_CYCLE_COUNTER(STAT_UserDataPT_RemoveData_GT);
+
+			if (this->GetSolver() != nullptr)
+			{
+				if (TInput* Input = this->GetProducerInputData_External())
+				{
+					// Track the particle for removal
+					Input->UserDataToRemove.Add(Handle.UniqueIdx());
+
+					// In case it was added/updated and then removed in the
+					// same frame, untrack the add/update
+					Input->UserDataToAdd.Remove(Handle.UniqueIdx());
+
+					// Successfully queued for removal
+					return true;
+				}
+			}
+
+			// Failed to queue for removal
+			return false;
+		}
+
+		// TParticleHandle is generalized here because it can be
+		// FRigidBodyHandle_Internal or FGeometryParticleHandle which have
+		// the same api...
+		template <typename TParticleHandle>
+		const TUserData* GetData_PT(const TParticleHandle& Handle) const
+		{
+			SCOPE_CYCLE_COUNTER(STAT_UserDataPT_GetData_PT);
+
+			const int32 Idx = Handle.UniqueIdx().Idx;
+			return UserDataMap_PT.IsValidIndex(Idx) ? UserDataMap_PT[Idx].Get() : nullptr;
+		}
+
+
+	protected:
+
+		virtual void OnPreSimulate_Internal() override
+		{
+			if (const TInput* Input = this->GetConsumerInput_Internal())
+			{
+				// Only proceed if the input has not yet been processed.
+				// 
+				// It's possible that we'll get multiple presimulate calls
+				// with the same input because the same input continues to be
+				// provided until a new one is received, so we cache the
+				// timestamp of the last processed input to make sure that we
+				// don't double-process it.
+				if (InputIdentifier_PT != Input->Identifier)
+				{
+					InputIdentifier_PT = Input->Identifier;
+
+					if (Input->UserDataToAdd.Num() > 0)
+					{
+						SCOPE_CYCLE_COUNTER(STAT_UserDataPT_UpdateData_PT);
+
+						// Move all the user data to the internal map
+						for (auto& Iter : Input->UserDataToAdd)
+						{
+							UserDataMap_PT.EmplaceAt(Iter.Key.Idx, MoveTemp(Iter.Value));
+						}
+					}
+
+					if (Input->UserDataToRemove.Num() > 0)
+					{
+						SCOPE_CYCLE_COUNTER(STAT_UserDataPT_RemoveData_PT);
+
+						// Delete user data that has been removed
+						for (const FUniqueIdx Idx : Input->UserDataToRemove)
+						{
+							UserDataMap_PT.RemoveAt(Idx.Idx);
+						}
+
+						// Shrink sparse array if we took elements off the end
+						UserDataMap_PT.Shrink();
+					}
+				}
+			}
+		}
+
+		// Identifier of the next input to be created on the game thread
+		int32 InputIdentifier_GT = 0;
+
+		// Identifier of the last input to be consumed on the physics thread
+		int32 InputIdentifier_PT = -1;
+
+		// Map of particle unique ids to user data
+		TSparseArray<TUniquePtr<TUserData>> UserDataMap_PT;
+	};
+}
diff --git a/Engine/Plugins/Experimental/ChaosUserDataPT/Source/ChaosUserData/Public/ChaosUserDataPTModule.h b/Engine/Plugins/Experimental/ChaosUserDataPT/Source/ChaosUserData/Public/ChaosUserDataPTModule.h
new file mode 100644
--- /dev/null
+++ b/Engine/Plugins/Experimental/ChaosUserDataPT/Source/ChaosUserData/Public/ChaosUserDataPTModule.h
@@ -0,0 +1,15 @@
+// Copyright Epic Games, Inc. All Rights Reserved.
+
+#pragma once
+
+#include "CoreMinimal.h"
+#include "Modules/ModuleManager.h"
+
+class FChaosUserDataPTModule : public IModuleInterface
+{
+public:
+
+	/** IModuleInterface implementation */
+	virtual void StartupModule() override;
+	virtual void ShutdownModule() override;
+};
\ No newline at end of file
diff --git a/Engine/Plugins/Experimental/ChaosUserDataPT/Source/ChaosUserData/Public/ChaosUserDataPTStats.h b/Engine/Plugins/Experimental/ChaosUserDataPT/Source/ChaosUserData/Public/ChaosUserDataPTStats.h
new file mode 100644
--- /dev/null
+++ b/Engine/Plugins/Experimental/ChaosUserDataPT/Source/ChaosUserData/Public/ChaosUserDataPTStats.h
@@ -0,0 +1,17 @@
+// Copyright Epic Games, Inc. All Rights Reserved.
+
+#pragma once
+
+#include "Stats/Stats.h"
+
+DECLARE_STATS_GROUP(TEXT("ChaosUserDataPT"), STATGROUP_ChaosUserDataPT, STATCAT_Advanced);
+
+namespace Chaos
+{
+	DECLARE_CYCLE_STAT_EXTERN(TEXT("ChaosUserDataPT::SetData_GT"), STAT_UserDataPT_SetData_GT, STATGROUP_ChaosUserDataPT, CHAOSUSERDATAPT_API);
+	DECLARE_CYCLE_STAT_EXTERN(TEXT("ChaosUserDataPT::RemoveData_GT"), STAT_UserDataPT_RemoveData_GT, STATGROUP_ChaosUserDataPT, CHAOSUSERDATAPT_API);
+	DECLARE_CYCLE_STAT_EXTERN(TEXT("ChaosUserDataPT::GetData_PT"), STAT_UserDataPT_GetData_PT, STATGROUP_ChaosUserDataPT, CHAOSUSERDATAPT_API);
+	DECLARE_CYCLE_STAT_EXTERN(TEXT("ChaosUserDataPT::OnPreSimulate::UpdateData"), STAT_UserDataPT_UpdateData_PT, STATGROUP_ChaosUserDataPT, CHAOSUSERDATAPT_API);
+	DECLARE_CYCLE_STAT_EXTERN(TEXT("ChaosUserDataPT::OnPreSimulate::RemoveData"), STAT_UserDataPT_RemoveData_PT, STATGROUP_ChaosUserDataPT, CHAOSUSERDATAPT_API);
+}
+
