Only remove elements from internal sparse array of ChaosUserDataPT if they actually exist

diff --git a/Engine/Plugins/Experimental/ChaosUserDataPT/Source/ChaosUserDataPT/Public/ChaosUserDataPT.h b/Engine/Plugins/Experimental/ChaosUserDataPT/Source/ChaosUserDataPT/Public/ChaosUserDataPT.h
--- a/Engine/Plugins/Experimental/ChaosUserDataPT/Source/ChaosUserDataPT/Public/ChaosUserDataPT.h
+++ b/Engine/Plugins/Experimental/ChaosUserDataPT/Source/ChaosUserDataPT/Public/ChaosUserDataPT.h
@@ -194,7 +194,10 @@ namespace Chaos
 						// Delete user data that has been removed
 						for (const FUniqueIdx Idx : Input->UserDataToRemove)
 						{
-							UserDataMap_PT.RemoveAt(Idx.Idx);
+							if (UserDataMap_PT.IsValidIndex(Idx.Idx))
+							{
+								UserDataMap_PT.RemoveAt(Idx.Idx);
+							}
 						}
 
 						// Shrink sparse array if we took elements off the end
