Add ability to clear all data from a PhysicsUserDataPT manager without having to queue up deletions for each individual particle.

diff --git a/Engine/Plugins/Experimental/ChaosUserDataPT/Source/ChaosUserDataPT/Private/ChaosUserDataPTStats.cpp b/Engine/Plugins/Experimental/ChaosUserDataPT/Source/ChaosUserDataPT/Private/ChaosUserDataPTStats.cpp
--- a/Engine/Plugins/Experimental/ChaosUserDataPT/Source/ChaosUserDataPT/Private/ChaosUserDataPTStats.cpp
+++ b/Engine/Plugins/Experimental/ChaosUserDataPT/Source/ChaosUserDataPT/Private/ChaosUserDataPTStats.cpp
@@ -9,4 +9,5 @@ namespace Chaos
 	DEFINE_STAT(STAT_UserDataPT_GetData_PT);
 	DEFINE_STAT(STAT_UserDataPT_UpdateData_PT);
 	DEFINE_STAT(STAT_UserDataPT_RemoveData_PT);
+	DEFINE_STAT(STAT_UserDataPT_ClearData_PT);
 }
diff --git a/Engine/Plugins/Experimental/ChaosUserDataPT/Source/ChaosUserDataPT/Public/ChaosUserDataPT.h b/Engine/Plugins/Experimental/ChaosUserDataPT/Source/ChaosUserDataPT/Public/ChaosUserDataPT.h
--- a/Engine/Plugins/Experimental/ChaosUserDataPT/Source/ChaosUserDataPT/Public/ChaosUserDataPT.h
+++ b/Engine/Plugins/Experimental/ChaosUserDataPT/Source/ChaosUserDataPT/Public/ChaosUserDataPT.h
@@ -61,6 +61,9 @@ namespace Chaos
 		// Set of particle unique indices for which to remove user data
 		TSet<FUniqueIdx> UserDataToRemove;
 
+		// Flag for clearing all data
+		bool bClear = false;
+
 		// Monotonically increasing identifier for the input object. Each
 		// newly constructed input will store and increment this;
 		int32 Identifier = -1;
@@ -69,6 +72,7 @@ namespace Chaos
 		{
 			UserDataToAdd.Reset();
 			UserDataToRemove.Reset();
+			bClear = false;
 			Identifier = -1;
 		}
 	};
@@ -133,8 +137,12 @@ namespace Chaos
 			{
 				if (TInput* Input = this->GetProducerInputData_External())
 				{
-					// Track the particle for removal
-					Input->UserDataToRemove.Add(Handle.UniqueIdx());
+					// Track the particle for removal. No point in doing so
+					// if we've already marked all data for clearing.
+					if (!Input->bClear)
+					{
+						Input->UserDataToRemove.Add(Handle.UniqueIdx());
+					}
 
 					// In case it was added/updated and then removed in the
 					// same frame, untrack the add/update
@@ -161,6 +169,23 @@ namespace Chaos
 			return UserDataMap_PT.IsValidIndex(Idx) ? UserDataMap_PT[Idx].Get() : nullptr;
 		}
 
+		// Clear all data
+		bool ClearData_GT()
+		{
+			if (TInput* Input = this->GetProducerInputData_External())
+			{
+				// All userdatas to add and remove can be cleared. Any userdatas that
+				// are added after this point will still be added.
+				Input->UserDataToAdd.Empty();
+				Input->UserDataToRemove.Empty();
+
+				// Mark the flag for clear all data
+				Input->bClear = true;
+			}
+
+			// Failed to queue removal
+			return false;
+		}
 
 	protected:
 
@@ -179,6 +204,17 @@ namespace Chaos
 				{
 					InputIdentifier_PT = Input->Identifier;
 
+					// Clear all data
+					if (Input->bClear)
+					{
+						SCOPE_CYCLE_COUNTER(STAT_UserDataPT_ClearData_PT);
+
+						// Empty the userdata map
+						UserDataMap_PT.Empty();
+						UserDataMap_PT.Shrink();
+					}
+
+					// Add new data
 					if (Input->UserDataToAdd.Num() > 0)
 					{
 						SCOPE_CYCLE_COUNTER(STAT_UserDataPT_UpdateData_PT);
@@ -190,6 +226,7 @@ namespace Chaos
 						}
 					}
 
+					// Remove old data
 					if (Input->UserDataToRemove.Num() > 0)
 					{
 						SCOPE_CYCLE_COUNTER(STAT_UserDataPT_RemoveData_PT);
diff --git a/Engine/Plugins/Experimental/ChaosUserDataPT/Source/ChaosUserDataPT/Public/ChaosUserDataPTStats.h b/Engine/Plugins/Experimental/ChaosUserDataPT/Source/ChaosUserDataPT/Public/ChaosUserDataPTStats.h
--- a/Engine/Plugins/Experimental/ChaosUserDataPT/Source/ChaosUserDataPT/Public/ChaosUserDataPTStats.h
+++ b/Engine/Plugins/Experimental/ChaosUserDataPT/Source/ChaosUserDataPT/Public/ChaosUserDataPTStats.h
@@ -13,5 +13,6 @@ namespace Chaos
 	DECLARE_CYCLE_STAT_EXTERN(TEXT("ChaosUserDataPT::GetData_PT"), STAT_UserDataPT_GetData_PT, STATGROUP_ChaosUserDataPT, CHAOSUSERDATAPT_API);
 	DECLARE_CYCLE_STAT_EXTERN(TEXT("ChaosUserDataPT::OnPreSimulate::UpdateData"), STAT_UserDataPT_UpdateData_PT, STATGROUP_ChaosUserDataPT, CHAOSUSERDATAPT_API);
 	DECLARE_CYCLE_STAT_EXTERN(TEXT("ChaosUserDataPT::OnPreSimulate::RemoveData"), STAT_UserDataPT_RemoveData_PT, STATGROUP_ChaosUserDataPT, CHAOSUSERDATAPT_API);
+	DECLARE_CYCLE_STAT_EXTERN(TEXT("ChaosUserDataPT::OnPreSimulate::ClearData"), STAT_UserDataPT_ClearData_PT, STATGROUP_ChaosUserDataPT, CHAOSUSERDATAPT_API);
 }
 
