Fix compilation error reported through UDN

diff --git a/Engine/Plugins/Experimental/ChaosUserDataPT/Source/ChaosUserDataPT/Public/ChaosUserDataPT.h b/Engine/Plugins/Experimental/ChaosUserDataPT/Source/ChaosUserDataPT/Public/ChaosUserDataPT.h
--- a/Engine/Plugins/Experimental/ChaosUserDataPT/Source/ChaosUserDataPT/Public/ChaosUserDataPT.h
+++ b/Engine/Plugins/Experimental/ChaosUserDataPT/Source/ChaosUserDataPT/Public/ChaosUserDataPT.h
@@ -93,7 +93,7 @@ namespace Chaos
 		{
 			SCOPE_CYCLE_COUNTER(STAT_UserDataPT_SetData_GT);
 
-			if (const FPhysicsSolverBase* Solver = this->GetSolver())
+			if (const FPhysicsSolverBase* SolverBase = this->GetSolver())
 			{
 				if (TInput* Input = this->GetProducerInputData_External())
 				{
