Generalize ChaosUserDataPT's GT API to use FPhysicsObject instead of specific physics handle types.
Update BuildingSMActor's usage of ChaosUserDataPT to use the FPhysicsObject from it's SimulatedComponent so that GCs will work too.

diff --git a/Engine/Plugins/Experimental/ChaosUserDataPT/Source/ChaosUserDataPT/Public/ChaosUserDataPT.h b/Engine/Plugins/Experimental/ChaosUserDataPT/Source/ChaosUserDataPT/Public/ChaosUserDataPT.h
--- a/Engine/Plugins/Experimental/ChaosUserDataPT/Source/ChaosUserDataPT/Public/ChaosUserDataPT.h
+++ b/Engine/Plugins/Experimental/ChaosUserDataPT/Source/ChaosUserDataPT/Public/ChaosUserDataPT.h
@@ -3,6 +3,8 @@
 #pragma once
 
 #include "PhysicsProxy/SingleParticlePhysicsProxy.h"
+#include "Chaos/PhysicsObject.h"
+#include "PhysicsEngine/PhysicsObjectExternalInterface.h"
 #include "ChaosUserDataPTStats.h"
 
 /*
@@ -88,8 +90,9 @@ namespace Chaos
 
 		virtual ~TUserDataManagerPT() { }
 
-		// Add or update user data associated with this particle handle
-		bool SetData_GT(const FRigidBodyHandle_External& Handle, const TUserData& UserData)
+		// Add or update user data associated with a particle handle
+		template <typename TParticleHandle>
+		bool SetData_GT(const TParticleHandle& Handle, const TUserData& UserData)
 		{
 			SCOPE_CYCLE_COUNTER(STAT_UserDataPT_SetData_GT);
 
@@ -119,8 +122,9 @@ namespace Chaos
 			return false;
 		}
 
-		// Remove user data associated with this particle handle
-		bool RemoveData_GT(const FRigidBodyHandle_External& Handle)
+		// Remove user data associated with a particle handle
+		template <typename TParticleHandle>
+		bool RemoveData_GT(const TParticleHandle& Handle)
 		{
 			SCOPE_CYCLE_COUNTER(STAT_UserDataPT_RemoveData_GT);
 
@@ -148,6 +152,32 @@ namespace Chaos
 			return false;
 		}
 
+		template <>
+		bool SetData_GT<Chaos::FPhysicsObjectHandle>(const Chaos::FPhysicsObjectHandle& Object, const TUserData& UserData)
+		{
+			// Get the game thread particle handle and set UserData on that
+			FLockedReadPhysicsObjectExternalInterface Interface = FPhysicsObjectExternalInterface::LockRead(Object);
+			if (const Chaos::FGeometryParticle* Particle = Interface->GetParticle(Object))
+			{
+				return SetData_GT(*Particle, UserData);
+			}
+
+			return false;
+		}
+
+		template <>
+		bool RemoveData_GT<Chaos::FPhysicsObjectHandle>(const Chaos::FPhysicsObjectHandle& Object)
+		{
+			// Get the game thread particle handle and remove UserData from it
+			FLockedReadPhysicsObjectExternalInterface Interface = FPhysicsObjectExternalInterface::LockRead(TArray<Chaos::FConstPhysicsObjectHandle>{ Object });
+			if (const Chaos::FGeometryParticle* Particle = Interface->GetParticle(Object))
+			{
+				return RemoveData_GT(*Particle);
+			}
+
+			return false;
+		}
+
 		// TParticleHandle is generalized here because it can be
 		// FRigidBodyHandle_Internal or FGeometryParticleHandle which have
 		// the same api...
diff --git a/Engine/Source/Runtime/Engine/Classes/PhysicsEngine/PhysicsObjectExternalInterface.h b/Engine/Source/Runtime/Engine/Classes/PhysicsEngine/PhysicsObjectExternalInterface.h
--- a/Engine/Source/Runtime/Engine/Classes/PhysicsEngine/PhysicsObjectExternalInterface.h
+++ b/Engine/Source/Runtime/Engine/Classes/PhysicsEngine/PhysicsObjectExternalInterface.h
@@ -47,6 +47,7 @@ using FLockedWritePhysicsObjectExternalInterface = FLockedPhysicsObjectExternalI
 class ENGINE_API FPhysicsObjectExternalInterface : public Chaos::FPhysicsObjectInterface, public PhysicsObjectPhysicsCoreInterface
 {
 public:
+	static FLockedReadPhysicsObjectExternalInterface LockRead(Chaos::FConstPhysicsObjectHandle InObject);
 	static FLockedReadPhysicsObjectExternalInterface LockRead(TArrayView<const Chaos::FConstPhysicsObjectHandle> InObjects);
 	static FLockedWritePhysicsObjectExternalInterface LockWrite(TArrayView<const Chaos::FPhysicsObjectHandle> InObjects);
 
diff --git a/Engine/Source/Runtime/Engine/Private/PhysicsEngine/PhysicsObjectExternalInterface.cpp b/Engine/Source/Runtime/Engine/Private/PhysicsEngine/PhysicsObjectExternalInterface.cpp
--- a/Engine/Source/Runtime/Engine/Private/PhysicsEngine/PhysicsObjectExternalInterface.cpp
+++ b/Engine/Source/Runtime/Engine/Private/PhysicsEngine/PhysicsObjectExternalInterface.cpp
@@ -1,6 +1,11 @@
 // Copyright Epic Games, Inc. All Rights Reserved.
 #include "PhysicsEngine/PhysicsObjectExternalInterface.h"
 
+FLockedReadPhysicsObjectExternalInterface FPhysicsObjectExternalInterface::LockRead(Chaos::FConstPhysicsObjectHandle InObject)
+{
+	return LockRead(TArray<Chaos::FConstPhysicsObjectHandle>{ InObject });
+}
+
 FLockedReadPhysicsObjectExternalInterface FPhysicsObjectExternalInterface::LockRead(TArrayView<const Chaos::FConstPhysicsObjectHandle> InObjects)
 {
 	return FLockedReadPhysicsObjectExternalInterface{ GetScene(InObjects), CreateReadInterface<Chaos::EThreadContext::External>() };
diff --git a/Engine/Source/Runtime/Experimental/Chaos/Private/Chaos/PhysicsObjectInterface.cpp b/Engine/Source/Runtime/Experimental/Chaos/Private/Chaos/PhysicsObjectInterface.cpp
--- a/Engine/Source/Runtime/Experimental/Chaos/Private/Chaos/PhysicsObjectInterface.cpp
+++ b/Engine/Source/Runtime/Experimental/Chaos/Private/Chaos/PhysicsObjectInterface.cpp
@@ -174,6 +174,12 @@ namespace Chaos
 		return {};
 	}
 
+	template<EThreadContext Id>
+	TThreadParticle<Id>* FReadPhysicsObjectInterface<Id>::GetParticle(const FConstPhysicsObjectHandle Object)
+	{
+		return Object->GetParticle<Id>();
+	}
+
 	template<EThreadContext Id>
 	TArray<TThreadParticle<Id>*> FReadPhysicsObjectInterface<Id>::GetAllParticles(TArrayView<const FConstPhysicsObjectHandle> InObjects)
 	{
diff --git a/Engine/Source/Runtime/Experimental/Chaos/Public/Chaos/PhysicsObjectInterface.h b/Engine/Source/Runtime/Experimental/Chaos/Public/Chaos/PhysicsObjectInterface.h
--- a/Engine/Source/Runtime/Experimental/Chaos/Public/Chaos/PhysicsObjectInterface.h
+++ b/Engine/Source/Runtime/Experimental/Chaos/Public/Chaos/PhysicsObjectInterface.h
@@ -65,6 +65,7 @@ namespace Chaos
 		FQuat GetR(const FConstPhysicsObjectHandle Object);
 		FSpatialAccelerationIdx GetSpatialIndex(const FConstPhysicsObjectHandle Object);
 
+		TThreadParticle<Id>* GetParticle(const FConstPhysicsObjectHandle Object);
 		TArray<TThreadParticle<Id>*> GetAllParticles(TArrayView<const FConstPhysicsObjectHandle> InObjects);
 		TArray<TThreadRigidParticle<Id>*> GetAllRigidParticles(TArrayView<const FConstPhysicsObjectHandle> InObjects);
 		TArray<FPerShapeData*> GetAllShapes(TArrayView<const FConstPhysicsObjectHandle> InObjects);
