No PCH build fix for Stats includes when STATS is 0

#rnx

diff --git a/Engine/Plugins/Experimental/ChaosUserDataPT/Source/ChaosUserDataPT/Public/ChaosUserDataPTStats.h b/Engine/Plugins/Experimental/ChaosUserDataPT/Source/ChaosUserDataPT/Public/ChaosUserDataPTStats.h
--- a/Engine/Plugins/Experimental/ChaosUserDataPT/Source/ChaosUserDataPT/Public/ChaosUserDataPTStats.h
+++ b/Engine/Plugins/Experimental/ChaosUserDataPT/Source/ChaosUserDataPT/Public/ChaosUserDataPTStats.h
@@ -2,9 +2,7 @@
 
 #pragma once
 
-#include "Stats/Stats2.h"
-
-enum class EStatFlags : uint8;
+#include "Stats/Stats.h"
 
 DECLARE_STATS_GROUP(TEXT("ChaosUserDataPT"), STATGROUP_ChaosUserDataPT, STATCAT_Advanced);
 
@@ -17,8 +15,3 @@ namespace Chaos
 	DECLARE_CYCLE_STAT_EXTERN(TEXT("ChaosUserDataPT::OnPreSimulate::RemoveData"), STAT_UserDataPT_RemoveData_PT, STATGROUP_ChaosUserDataPT, CHAOSUSERDATAPT_API);
 	DECLARE_CYCLE_STAT_EXTERN(TEXT("ChaosUserDataPT::OnPreSimulate::ClearData"), STAT_UserDataPT_ClearData_PT, STATGROUP_ChaosUserDataPT, CHAOSUSERDATAPT_API);
 }
-
-
-#if UE_ENABLE_INCLUDE_ORDER_DEPRECATED_IN_5_2
-#include "Stats/Stats.h"
-#endif
