TUserDataManagerPT: add GetData_GT for optional game-thread read-back

Adds FUserDataPTConfig with bGetData_GT flag. When enabled, SetData_GT
mirrors writes to a GT-side sparse array so GetData_GT can return the
last set value without a physics-thread round trip. When disabled (the
default), GetData_GT returns nullptr.

Other changes
- SetData_GT rvalue-ref overload (avoids redundant copy).
- ClearData_GT now returns true on success (was unconditionally false).
- InputIdentifier changed from int32 to uint32 (signed overflow UB).
- bResizeOnClear renamed to bResize; now reset in TInput::Reset().
- TFunction -> TFunctionRef in VisitAllData_PT (avoids PT heap alloc).
- New stat counter STAT_UserDataPT_GetData_GT.
- New test coverage for GetData_GT, multi-particle, same-frame
set+remove, VisitAllData_PT.

#Synced-CL 52522074

diff --git a/Engine/Plugins/Experimental/Buoyancy/Source/Runtime/Private/BuoyancySubsystem.cpp b/Engine/Plugins/Experimental/Buoyancy/Source/Runtime/Private/BuoyancySubsystem.cpp
--- a/Engine/Plugins/Experimental/Buoyancy/Source/Runtime/Private/BuoyancySubsystem.cpp
+++ b/Engine/Plugins/Experimental/Buoyancy/Source/Runtime/Private/BuoyancySubsystem.cpp
@@ -911,7 +911,7 @@ void FBuoyancySubsystemSimCallback::TrackInteractions(
 	// The spline userdatapt manager contains a PT array of all water splines indexed on
 	// water particle unique index. Therefore, we can use it to loop over all water body
 	// particles and check their midphases.
-	SplineData->VisitAllData_PT([this, &PBDSolver, &Evolution, &MidPhaseAccessor](Chaos::FUniqueIdx ParticleIdx, const TSharedPtr<FBuoyancyWaterSplineData>& WaterData)
+	SplineData->VisitData_PT([this, &PBDSolver, &Evolution, &MidPhaseAccessor](Chaos::FUniqueIdx ParticleIdx, const TSharedPtr<FBuoyancyWaterSplineData>& WaterData)
 	{
 		// If the WaterData shared ptr is invalid, skip this one
 		if (WaterData.IsValid() == false)
diff --git a/Engine/Plugins/Experimental/ChaosUserDataPT/Source/ChaosUserDataPT/Private/ChaosUserDataPTStats.cpp b/Engine/Plugins/Experimental/ChaosUserDataPT/Source/ChaosUserDataPT/Private/ChaosUserDataPTStats.cpp
--- a/Engine/Plugins/Experimental/ChaosUserDataPT/Source/ChaosUserDataPT/Private/ChaosUserDataPTStats.cpp
+++ b/Engine/Plugins/Experimental/ChaosUserDataPT/Source/ChaosUserDataPT/Private/ChaosUserDataPTStats.cpp
@@ -8,6 +8,7 @@ namespace Chaos
 	DEFINE_STAT(STAT_UserDataPT_SetData_GT);
 	DEFINE_STAT(STAT_UserDataPT_RemoveData_GT);
 	DEFINE_STAT(STAT_UserDataPT_GetData_PT);
+	DEFINE_STAT(STAT_UserDataPT_GetData_GT);
 	DEFINE_STAT(STAT_UserDataPT_UpdateData_PT);
 	DEFINE_STAT(STAT_UserDataPT_RemoveData_PT);
 	DEFINE_STAT(STAT_UserDataPT_ClearData_PT);
diff --git a/Engine/Plugins/Experimental/ChaosUserDataPT/Source/ChaosUserDataPT/Public/ChaosUserDataPT.h b/Engine/Plugins/Experimental/ChaosUserDataPT/Source/ChaosUserDataPT/Public/ChaosUserDataPT.h
--- a/Engine/Plugins/Experimental/ChaosUserDataPT/Source/ChaosUserDataPT/Public/ChaosUserDataPT.h
+++ b/Engine/Plugins/Experimental/ChaosUserDataPT/Source/ChaosUserDataPT/Public/ChaosUserDataPT.h
@@ -3,9 +3,11 @@
 #pragma once
 
 #include "PhysicsProxy/SingleParticlePhysicsProxy.h"
+#include "PBDRigidsSolver.h"
 #include "Chaos/PhysicsObject.h"
 #include "PhysicsEngine/PhysicsObjectExternalInterface.h"
 #include "ChaosUserDataPTStats.h"
+#include <limits>
 
 /*
 
@@ -34,13 +36,63 @@
 
 namespace Chaos
 {
+	template <typename TUserData>
+	struct TUserDataManagerPTInput;
+
+	template <typename TUserData>
+	class TUserDataManagerPT;
+
+	// Configuration for TUserDataManagerPT behavior.
+	// Passed at construction time and stored as a const member.
+	struct FUserDataPTConfig
+	{
+		// When true, TUserDataManagerPT will cache userdata on GT as well as PT,
+		// so that GT users can get the last set value. Useful when they want
+		// to make incremental changes, but incurs a higher memory footprint.
+		bool bGetData_GT = false;
+	};
+
 	// TUserDataManagerPTInput
 	//
 	// Input is a collection of pointers to new and updated userdata objects
 	// to be sent to the physics thread
+	//
+	// NOTE: On deprecation strategy here.
+	// The members of this struct should never have been public, so they have
+	// all been deprecated. Private duplicates are now in use, but those will
+	// be renamed to the original names once the deprecated versions are removed.
+	//
+	// Copies are needed so that during the deprecation phase the internal
+	// TUserDataManagerPT code doesn't trigger deprecation warnings for valid
+	// operations.
 	template <typename TUserData>
 	struct TUserDataManagerPTInput : public FSimCallbackInput
 	{
+		UE_DEPRECATED(5.8, "UserDataToAdd no longer used, will be removed.")
+		mutable TMap<FUniqueIdx, TUniquePtr<TUserData>> UserDataToAdd;
+
+		UE_DEPRECATED(5.8, "UserDataToRemove no longer used, will be removed.")
+		TSet<FUniqueIdx> UserDataToRemove;
+
+		UE_DEPRECATED(5.8, "bClear no longer used, will be removed.")
+		bool bClear = false;
+
+		UE_DEPRECATED(5.8, "bResizeOnClear no longer used, will be removed.")
+		bool bResizeOnClear = false;
+
+		UE_DEPRECATED(5.8, "Identifier no longer used, will be removed.")
+		uint32 Identifier = std::numeric_limits<uint32>::max();
+
+		void Reset()
+		{
+			UserDataToAddPrivate.Reset();
+			UserDataToRemovePrivate.Reset();
+			bClearPrivate = false;
+			bResize = false;
+			IdentifierPrivate = std::numeric_limits<uint32>::max();
+		}
+
+	private:
 		// Map of particle unique indices to user data ptrs
 		// 
 		// NOTE: This is marked mutable because the userdata objects must be
@@ -49,28 +101,23 @@ namespace Chaos
 		// context. This might be frowned on, but since TUserDataManagerPTInput is a
 		// class which is only used internally to TUserDataManagerPT, an argument
 		// could be made either way.
-		mutable TMap<FUniqueIdx, TUniquePtr<TUserData>> UserDataToAdd;
+		mutable TMap<FUniqueIdx, TUniquePtr<TUserData>> UserDataToAddPrivate;
 
 		// Set of particle unique indices for which to remove user data
-		TSet<FUniqueIdx> UserDataToRemove;
+		TSet<FUniqueIdx> UserDataToRemovePrivate;
 
 		// Flag for clearing all data
-		bool bClear = false;
+		bool bClearPrivate = false;
 
-		// Flag for whether or not data resize is allowed from changes
-		bool bResizeOnClear = true;
+		// Flag for resizing container to fit data
+		bool bResize = false;
 
 		// Monotonically increasing identifier for the input object. Each
 		// newly constructed input will store and increment this;
-		int32 Identifier = -1;
+		uint32 IdentifierPrivate = std::numeric_limits<uint32>::max();
 
-		void Reset()
-		{
-			UserDataToAdd.Reset();
-			UserDataToRemove.Reset();
-			bClear = false;
-			Identifier = -1;
-		}
+		// NOTE: In UE 5.8, mark all contents of this struct private
+		friend class TUserDataManagerPT<TUserData>;
 	};
 
 	// TUserDataManagerPT
@@ -91,11 +138,24 @@ namespace Chaos
 
 	public:
 
+		explicit TUserDataManagerPT(const FUserDataPTConfig& InConfig = FUserDataPTConfig())
+			: Config(InConfig)
+		{
+		}
+
 		virtual ~TUserDataManagerPT() = default;
 
-		// Add or update user data associated with a particle handle
+		// Add or update user data associated with a particle handle (copy)
 		template <typename TParticleHandle>
 		bool SetData_GT(const TParticleHandle& Handle, const TUserData& UserData)
+		{
+			TUserData Copy(UserData);
+			return SetData_GT(Handle, MoveTemp(Copy));
+		}
+
+		// Add or update user data associated with a particle handle (move)
+		template <typename TParticleHandle>
+		bool SetData_GT(const TParticleHandle& Handle, TUserData&& UserData)
 		{
 			SCOPE_CYCLE_COUNTER(STAT_UserDataPT_SetData_GT);
 
@@ -103,17 +163,23 @@ namespace Chaos
 			{
 				if (TInput* Input = this->GetProducerInputData_External())
 				{
+					// Do GT data caching if configured
+					if (Config.bGetData_GT)
+					{
+						UserDataMap_GT.EmplaceAt(Handle.UniqueIdx().Idx, MakeUnique<TUserData>(UserData));
+					}
+
 					// Add the data to the map to be sent to physics thread
-					Input->UserDataToAdd.Emplace(Handle.UniqueIdx(), MakeUnique<TUserData>(UserData));
+					Input->UserDataToAddPrivate.Emplace(Handle.UniqueIdx(), MakeUnique<TUserData>(MoveTemp(UserData)));
 
 					// In case it was removed and then added again in the same
 					// frame, untrack this particle for data removal
-					Input->UserDataToRemove.Remove(Handle.UniqueIdx());
+					Input->UserDataToRemovePrivate.Remove(Handle.UniqueIdx());
 
-					// If this is a new input, set it's identifier
-					if (Input->Identifier == -1)
+					// If this is a new input, set it's IdentifierPrivate
+					if (Input->IdentifierPrivate == std::numeric_limits<uint32>::max())
 					{
-						Input->Identifier = InputIdentifier_GT++;
+						Input->IdentifierPrivate = InputIdentifier_GT++;
 					}
 
 					// Successfully queued for add/update
@@ -125,6 +191,22 @@ namespace Chaos
 			return false;
 		}
 
+		// Get user data from the GT cache. Returns nullptr if
+		// Config.bGetData_GT is false or the particle has no data.
+		template <typename TParticleHandle>
+		const TUserData* GetData_GT(const TParticleHandle& Handle) const
+		{
+			SCOPE_CYCLE_COUNTER(STAT_UserDataPT_GetData_GT);
+
+			if (!Config.bGetData_GT)
+			{
+				return nullptr;
+			}
+
+			const int32 Idx = Handle.UniqueIdx().Idx;
+			return UserDataMap_GT.IsValidIndex(Idx) ? UserDataMap_GT[Idx].Get() : nullptr;
+		}
+
 		// Remove user data associated with a particle handle
 		template <typename TParticleHandle>
 		bool RemoveData_GT(const TParticleHandle& Handle)
@@ -137,14 +219,23 @@ namespace Chaos
 				{
 					// Track the particle for removal. No point in doing so
 					// if we've already marked all data for clearing.
-					if (!Input->bClear)
+					if (!Input->bClearPrivate)
 					{
-						Input->UserDataToRemove.Add(Handle.UniqueIdx());
+						Input->UserDataToRemovePrivate.Add(Handle.UniqueIdx());
 					}
 
 					// In case it was added/updated and then removed in the
 					// same frame, untrack the add/update
-					Input->UserDataToAdd.Remove(Handle.UniqueIdx());
+					Input->UserDataToAddPrivate.Remove(Handle.UniqueIdx());
+
+					// Remove GT data cache if configured
+					if (Config.bGetData_GT)
+					{
+						if (UserDataMap_GT.IsValidIndex(Handle.UniqueIdx().Idx))
+						{
+							UserDataMap_GT.RemoveAt(Handle.UniqueIdx().Idx);
+						}
+					}
 
 					// Successfully queued for removal
 					return true;
@@ -172,7 +263,7 @@ namespace Chaos
 		bool RemoveData_GT<Chaos::FPhysicsObjectHandle>(const Chaos::FPhysicsObjectHandle& Object)
 		{
 			// Get the game thread particle handle and remove UserData from it
-			FLockedReadPhysicsObjectExternalInterface Interface = FPhysicsObjectExternalInterface::LockRead(TArray<Chaos::FConstPhysicsObjectHandle>{ Object });
+			FLockedReadPhysicsObjectExternalInterface Interface = FPhysicsObjectExternalInterface::LockRead(Object);
 			if (const Chaos::FGeometryParticle* Particle = Interface->GetParticle(Object))
 			{
 				return RemoveData_GT(*Particle);
@@ -194,11 +285,10 @@ namespace Chaos
 		}
 
 		// Execute a lambda on each data entry
-		void VisitAllData_PT(const TFunction<void(Chaos::FUniqueIdx, const TUserData&)>& Lambda)
+		void VisitData_PT(TFunctionRef<void(Chaos::FUniqueIdx, const TUserData&)> Lambda)
 		{
-			for (auto It = UserDataMap_PT.CreateIterator(); It; ++It)
+			for (auto It = UserDataMap_PT.CreateConstIterator(); It; ++It)
 			{
-				// If the unique ptr has data
 				if (It->IsValid())
 				{
 					Lambda(Chaos::FUniqueIdx(It.GetIndex()), *It->Get());
@@ -206,6 +296,12 @@ namespace Chaos
 			}
 		}
 
+		UE_DEPRECATED(5.8, "VisitAllData_PT will be removed, use VisitData_PT instead.")
+		void VisitAllData_PT(const TFunction<void(Chaos::FUniqueIdx, const TUserData&)>& Lambda)
+		{
+			VisitData_PT(Lambda);
+		}
+
 		// Clear all data
 		bool ClearData_GT(const bool bResize = true)
 		{
@@ -213,17 +309,35 @@ namespace Chaos
 			{
 				// All userdatas to add and remove can be cleared. Any userdatas that
 				// are added after this point will still be added.
-				Input->UserDataToAdd.Empty();
-				Input->UserDataToRemove.Empty();
+				Input->UserDataToAddPrivate.Empty();
+				Input->UserDataToRemovePrivate.Empty();
 
-				// Mark the flag for clear all data
-				Input->bClear = true;
+				// Mark the flags for clearing all data and resizing container
+				Input->bClearPrivate = true;
+				Input->bResize = bResize;
 
-				// Mark whether or not we're allowed to resize
-				Input->bResizeOnClear = bResize;
+				// If this is a new input, set it's IdentifierPrivate
+				if (Input->IdentifierPrivate == std::numeric_limits<uint32>::max())
+				{
+					Input->IdentifierPrivate = InputIdentifier_GT++;
+				}
+
+				// Empty the GT data cache if configured
+				if (Config.bGetData_GT)
+				{
+					UserDataMap_GT.Empty();
+
+					if (bResize)
+					{
+						UserDataMap_GT.Shrink();
+					}
+				}
+
+				// Succeeded
+				return true;
 			}
 
-			// Failed to queue removal
+			// Failed to queue data for clearing 
 			return false;
 		}
 
@@ -233,8 +347,47 @@ namespace Chaos
 			return StaticName;
 		}
 
+		// Runtime configuration (immutable after construction)
+		const FUserDataPTConfig Config;
+
 	protected:
 
+		virtual void OnPostInitialize_Internal() override
+		{
+			if (Config.bGetData_GT)
+			{
+				if (FPBDRigidsSolver* RigidsSolver = static_cast<FPBDRigidsSolver*>(this->GetSolver()))
+				{
+PRAGMA_DISABLE_INTERNAL_WARNINGS
+					ParticleUnregisteredGTHandle = RigidsSolver->AddParticleUnregisteredGTCallback(
+						FSolverParticleUnregisteredGT::FDelegate::CreateLambda(
+						[this](FSingleParticlePhysicsProxy* /*Proxy*/, FUniqueIdx UniqueIdx)
+						{
+							if (UserDataMap_GT.IsValidIndex(UniqueIdx.Idx))
+							{
+								UserDataMap_GT.RemoveAt(UniqueIdx.Idx);
+							}
+						})
+					);
+PRAGMA_ENABLE_INTERNAL_WARNINGS
+				}
+			}
+		}
+
+		virtual void OnPreShutdown_External() override
+		{
+			if (ParticleUnregisteredGTHandle.IsValid())
+			{
+				if (FPBDRigidsSolver* RigidsSolver = static_cast<FPBDRigidsSolver*>(this->GetSolver()))
+				{
+PRAGMA_DISABLE_INTERNAL_WARNINGS
+					RigidsSolver->RemoveParticleUnregisteredGTCallback(ParticleUnregisteredGTHandle);
+					ParticleUnregisteredGTHandle.Reset();
+PRAGMA_ENABLE_INTERNAL_WARNINGS
+				}
+			}
+		}
+
 		virtual void OnPreSimulate_Internal() override
 		{
 			if (const TInput* Input = this->GetConsumerInput_Internal())
@@ -246,44 +399,38 @@ namespace Chaos
 				// provided until a new one is received, so we cache the
 				// timestamp of the last processed input to make sure that we
 				// don't double-process it.
-				if (InputIdentifier_PT != Input->Identifier)
+				if (InputIdentifier_PT != Input->IdentifierPrivate)
 				{
-					InputIdentifier_PT = Input->Identifier;
+					InputIdentifier_PT = Input->IdentifierPrivate;
 
 					// Clear all data
-					if (Input->bClear)
+					if (Input->bClearPrivate)
 					{
 						SCOPE_CYCLE_COUNTER(STAT_UserDataPT_ClearData_PT);
 
 						// Empty the userdata map
 						UserDataMap_PT.Empty();
-
-						// If resize is allowed, do it
-						if (Input->bResizeOnClear)
-						{
-							UserDataMap_PT.Shrink();
-						}
 					}
 
 					// Add new data
-					if (Input->UserDataToAdd.Num() > 0)
+					if (Input->UserDataToAddPrivate.Num() > 0)
 					{
 						SCOPE_CYCLE_COUNTER(STAT_UserDataPT_UpdateData_PT);
 
 						// Move all the user data to the internal map
-						for (auto& Iter : Input->UserDataToAdd)
+						for (auto& Iter : Input->UserDataToAddPrivate)
 						{
 							UserDataMap_PT.EmplaceAt(Iter.Key.Idx, MoveTemp(Iter.Value));
 						}
 					}
 
 					// Remove old data
-					if (Input->UserDataToRemove.Num() > 0)
+					if (Input->UserDataToRemovePrivate.Num() > 0)
 					{
 						SCOPE_CYCLE_COUNTER(STAT_UserDataPT_RemoveData_PT);
 
 						// Delete user data that has been removed
-						for (const FUniqueIdx Idx : Input->UserDataToRemove)
+						for (const FUniqueIdx Idx : Input->UserDataToRemovePrivate)
 						{
 							if (UserDataMap_PT.IsValidIndex(Idx.Idx))
 							{
@@ -291,7 +438,12 @@ namespace Chaos
 							}
 						}
 
-						// Shrink sparse array if we took elements off the end
+					}
+
+					// Resize container
+					if (Input->bResize)
+					{
+						// Shrink sparse array to fit the data
 						UserDataMap_PT.Shrink();
 					}
 				}
@@ -319,12 +471,19 @@ namespace Chaos
 		}
 
 		// Identifier of the next input to be created on the game thread
-		int32 InputIdentifier_GT = 0;
+		uint32 InputIdentifier_GT = 0;
 
 		// Identifier of the last input to be consumed on the physics thread
-		int32 InputIdentifier_PT = -1;
+		uint32 InputIdentifier_PT = std::numeric_limits<uint32>::max();
 
-		// Map of particle unique ids to user data
+		// Map of particle unique ids to user data, for PT access
 		TSparseArray<TUniquePtr<TUserData>> UserDataMap_PT;
+
+		// Map of particle unique ids to user data, for optional GT access.
+		// Only populated when Config.bGetData_GT is true; otherwise stays empty.
+		TSparseArray<TUniquePtr<TUserData>> UserDataMap_GT;
+
+		// Delegate handle for GT particle unregistration cleanup
+		FDelegateHandle ParticleUnregisteredGTHandle;
 	};
 }
diff --git a/Engine/Plugins/Experimental/ChaosUserDataPT/Source/ChaosUserDataPT/Public/ChaosUserDataPTStats.h b/Engine/Plugins/Experimental/ChaosUserDataPT/Source/ChaosUserDataPT/Public/ChaosUserDataPTStats.h
--- a/Engine/Plugins/Experimental/ChaosUserDataPT/Source/ChaosUserDataPT/Public/ChaosUserDataPTStats.h
+++ b/Engine/Plugins/Experimental/ChaosUserDataPT/Source/ChaosUserDataPT/Public/ChaosUserDataPTStats.h
@@ -11,6 +11,7 @@ namespace Chaos
 	DECLARE_CYCLE_STAT_EXTERN(TEXT("ChaosUserDataPT::SetData_GT"), STAT_UserDataPT_SetData_GT, STATGROUP_ChaosUserDataPT, CHAOSUSERDATAPT_API);
 	DECLARE_CYCLE_STAT_EXTERN(TEXT("ChaosUserDataPT::RemoveData_GT"), STAT_UserDataPT_RemoveData_GT, STATGROUP_ChaosUserDataPT, CHAOSUSERDATAPT_API);
 	DECLARE_CYCLE_STAT_EXTERN(TEXT("ChaosUserDataPT::GetData_PT"), STAT_UserDataPT_GetData_PT, STATGROUP_ChaosUserDataPT, CHAOSUSERDATAPT_API);
+	DECLARE_CYCLE_STAT_EXTERN(TEXT("ChaosUserDataPT::GetData_GT"), STAT_UserDataPT_GetData_GT, STATGROUP_ChaosUserDataPT, CHAOSUSERDATAPT_API);
 	DECLARE_CYCLE_STAT_EXTERN(TEXT("ChaosUserDataPT::OnPreSimulate::UpdateData"), STAT_UserDataPT_UpdateData_PT, STATGROUP_ChaosUserDataPT, CHAOSUSERDATAPT_API);
 	DECLARE_CYCLE_STAT_EXTERN(TEXT("ChaosUserDataPT::OnPreSimulate::RemoveData"), STAT_UserDataPT_RemoveData_PT, STATGROUP_ChaosUserDataPT, CHAOSUSERDATAPT_API);
 	DECLARE_CYCLE_STAT_EXTERN(TEXT("ChaosUserDataPT::OnPreSimulate::ClearData"), STAT_UserDataPT_ClearData_PT, STATGROUP_ChaosUserDataPT, CHAOSUSERDATAPT_API);
diff --git a/Engine/Source/Programs/ChaosUserDataPTTests/Private/ChaosUserDataPTTests.cpp b/Engine/Source/Programs/ChaosUserDataPTTests/Private/ChaosUserDataPTTests.cpp
--- a/Engine/Source/Programs/ChaosUserDataPTTests/Private/ChaosUserDataPTTests.cpp
+++ b/Engine/Source/Programs/ChaosUserDataPTTests/Private/ChaosUserDataPTTests.cpp
@@ -9,6 +9,13 @@
 
 class FTestUserData : public Chaos::TUserDataManagerPT<FString> { };
 
+class FTestUserData_GetDataGT : public Chaos::TUserDataManagerPT<FString>
+{
+public:
+	FTestUserData_GetDataGT() : TUserDataManagerPT(Chaos::FUserDataPTConfig{ .bGetData_GT = true }) {}
+};
+
+
 // Advance a solver once with the provided DeltaTime, then wait for any async tasks to finish before continuing
 void AdvanceAndWait(Chaos::FPBDRigidsSolver* InSolver, float DeltaTime)
 {
@@ -24,6 +31,7 @@ TEST_CASE("ChaosUserDataPT", "[integration]")
 	const float DeltaTime = 1.f;
 	const FString TestString1 = TEXT("TestData1");
 	const FString TestString2 = TEXT("TestData2");
+	const FString TestString3 = TEXT("TestData3");
 
 	// Create a solver in the solvers module
 	FChaosSolversModule* Module = FChaosSolversModule::GetModule();
@@ -192,8 +200,255 @@ TEST_CASE("ChaosUserDataPT", "[integration]")
 		});
 		AdvanceAndWait(Solver, DeltaTime);
 	}
+	
+	SECTION("Multiple particles with different data")
+	{
+		TestUserData->SetData_GT(HandleExternal0, TestString1);
+		TestUserData->SetData_GT(HandleExternal1, TestString2);
+		TestUserData->SetData_GT(HandleExternal2, TestString3);
+		AdvanceAndWait(Solver, DeltaTime);
+		AdvanceAndWait(Solver, DeltaTime);
+
+		Solver->EnqueueCommandImmediate([&]()
+		{
+			REQUIRE(*TestUserData->GetData_PT(*Proxy0->GetPhysicsThreadAPI()) == TestString1);
+			REQUIRE(*TestUserData->GetData_PT(*Proxy1->GetPhysicsThreadAPI()) == TestString2);
+			REQUIRE(*TestUserData->GetData_PT(*Proxy2->GetPhysicsThreadAPI()) == TestString3);
+		});
+		AdvanceAndWait(Solver, DeltaTime);
+	}
+	
+	SECTION("Set and remove in same frame results in no data on PT")
+	{
+		TestUserData->SetData_GT(HandleExternal0, TestString1);
+		TestUserData->RemoveData_GT(HandleExternal0);
+		AdvanceAndWait(Solver, DeltaTime);
+		AdvanceAndWait(Solver, DeltaTime);
+
+		Solver->EnqueueCommandImmediate([&]()
+		{
+			REQUIRE(TestUserData->GetData_PT(*Proxy0->GetPhysicsThreadAPI()) == nullptr);
+		});
+		AdvanceAndWait(Solver, DeltaTime);
+	}
+
+	SECTION("Remove and re-add in same frame preserves new value")
+	{
+		TestUserData->SetData_GT(HandleExternal0, TestString1);
+		AdvanceAndWait(Solver, DeltaTime);
+		AdvanceAndWait(Solver, DeltaTime);
+
+		TestUserData->RemoveData_GT(HandleExternal0);
+		TestUserData->SetData_GT(HandleExternal0, TestString2);
+		AdvanceAndWait(Solver, DeltaTime);
+
+		Solver->EnqueueCommandImmediate([&]()
+		{
+			REQUIRE(*TestUserData->GetData_PT(*Proxy0->GetPhysicsThreadAPI()) == TestString2);
+		});
+		AdvanceAndWait(Solver, DeltaTime);
+	}
+
+	SECTION("VisitData_PT iterates over all stored entries")
+	{
+		TestUserData->SetData_GT(HandleExternal0, TestString1);
+		TestUserData->SetData_GT(HandleExternal1, TestString2);
+		TestUserData->SetData_GT(HandleExternal2, TestString3);
+		AdvanceAndWait(Solver, DeltaTime);
+		AdvanceAndWait(Solver, DeltaTime);
+
+		Solver->EnqueueCommandImmediate([&]()
+		{
+			TMap<int32, FString> Visited;
+			TestUserData->VisitData_PT([&Visited](Chaos::FUniqueIdx Idx, const FString& Data)
+			{
+				Visited.Add(Idx.Idx, Data);
+			});
+
+			REQUIRE(Visited.Num() == 3);
+			REQUIRE(Visited[Proxy0->GetPhysicsThreadAPI()->UniqueIdx().Idx] == TestString1);
+			REQUIRE(Visited[Proxy1->GetPhysicsThreadAPI()->UniqueIdx().Idx] == TestString2);
+			REQUIRE(Visited[Proxy2->GetPhysicsThreadAPI()->UniqueIdx().Idx] == TestString3);
+		});
+		AdvanceAndWait(Solver, DeltaTime);
+	}
+
+	SECTION("GetData_GT returns nullptr when bGetData_GT is false")
+	{
+		TestUserData->SetData_GT(HandleExternal0, TestString1);
+		REQUIRE(TestUserData->GetData_GT(HandleExternal0) == nullptr);
+	}
 
 	Solver->WaitOnPendingTasks_External();
 	Module->DestroySolver(Solver);
 }
 
+
+TEST_CASE("FTestUserData_GetDataGT", "[integration]")
+{
+	const float DeltaTime = 1.f;
+	const FString TestString1 = TEXT("TestData1");
+	const FString TestString2 = TEXT("TestData2");
+
+	FChaosSolversModule* Module = FChaosSolversModule::GetModule();
+	Chaos::FPBDRigidsSolver* Solver = Module->CreateSolver(nullptr, /*AsyncDt=*/DeltaTime, Chaos::EThreadingMode::TaskGraph);
+	FTestUserData_GetDataGT* TestUserData = Solver->CreateAndRegisterSimCallbackObject_External<FTestUserData_GetDataGT>();
+
+	auto BoxGeom = Chaos::FImplicitObjectPtr(new Chaos::TBox<Chaos::FReal, 3>(Chaos::FVec3(-1, -1, -1), Chaos::FVec3(1, 1, 1)));
+
+	Chaos::FSingleParticlePhysicsProxy* Proxy0 = Chaos::FSingleParticlePhysicsProxy::Create(Chaos::FGeometryParticle::CreateParticle());
+	Chaos::FSingleParticlePhysicsProxy* Proxy1 = Chaos::FSingleParticlePhysicsProxy::Create(Chaos::FGeometryParticle::CreateParticle());
+	Chaos::FRigidBodyHandle_External& HandleExternal0 = Proxy0->GetGameThreadAPI();
+	Chaos::FRigidBodyHandle_External& HandleExternal1 = Proxy1->GetGameThreadAPI();
+	HandleExternal0.SetGeometry(BoxGeom);
+	HandleExternal1.SetGeometry(BoxGeom);
+	Solver->RegisterObject(Proxy0);
+	Solver->RegisterObject(Proxy1);
+
+	AdvanceAndWait(Solver, DeltaTime);
+	AdvanceAndWait(Solver, DeltaTime);
+
+	SECTION("Data updates are accessible from GT with bGetData_GT = true")
+	{
+		const FString* StringPtr = TestUserData->GetData_GT(HandleExternal0);
+		REQUIRE(StringPtr == nullptr);
+
+		TestUserData->SetData_GT(HandleExternal0, TestString1);
+		StringPtr = TestUserData->GetData_GT(HandleExternal0);
+		REQUIRE(*StringPtr == TestString1);
+
+		AdvanceAndWait(Solver, DeltaTime);
+
+		TestUserData->SetData_GT(HandleExternal0, TestString2);
+		StringPtr = TestUserData->GetData_GT(HandleExternal0);
+		REQUIRE(*StringPtr == TestString2);
+	}
+
+	SECTION("GT data is removed properly")
+	{
+		TestUserData->SetData_GT(HandleExternal0, TestString1);
+		TestUserData->SetData_GT(HandleExternal1, TestString2);
+		const FString* StringPtr0 = TestUserData->GetData_GT(HandleExternal0);
+		const FString* StringPtr1 = TestUserData->GetData_GT(HandleExternal1);
+		REQUIRE(*StringPtr0 == TestString1);
+		REQUIRE(*StringPtr1 == TestString2);
+
+		TestUserData->RemoveData_GT(HandleExternal0);
+		StringPtr0 = TestUserData->GetData_GT(HandleExternal0);
+		StringPtr1 = TestUserData->GetData_GT(HandleExternal1);
+		REQUIRE(StringPtr0 == nullptr);
+		REQUIRE(*StringPtr1 == TestString2);
+
+		TestUserData->ClearData_GT();
+		StringPtr0 = TestUserData->GetData_GT(HandleExternal0);
+		StringPtr1 = TestUserData->GetData_GT(HandleExternal1);
+		REQUIRE(StringPtr0 == nullptr);
+		REQUIRE(StringPtr1 == nullptr);
+	}
+
+	SECTION("GT data also propagates to PT")
+	{
+		TestUserData->SetData_GT(HandleExternal0, TestString1);
+		AdvanceAndWait(Solver, DeltaTime);
+
+		REQUIRE(*TestUserData->GetData_GT(HandleExternal0) == TestString1);
+
+		Solver->EnqueueCommandImmediate([&]()
+		{
+			REQUIRE(*TestUserData->GetData_PT(*Proxy0->GetPhysicsThreadAPI()) == TestString1);
+		});
+		AdvanceAndWait(Solver, DeltaTime);
+	}
+
+	SECTION("GT data is immediately removed when particle is unregistered")
+	{
+		// Set data on both particles
+		TestUserData->SetData_GT(HandleExternal0, TestString1);
+		TestUserData->SetData_GT(HandleExternal1, TestString2);
+		REQUIRE(*TestUserData->GetData_GT(HandleExternal0) == TestString1);
+		REQUIRE(*TestUserData->GetData_GT(HandleExternal1) == TestString2);
+
+		// Save the unique index before unregistering, since the handle
+		// will be invalid after UnregisterObject
+		struct FMockHandle
+		{
+			FMockHandle(Chaos::FUniqueIdx InUniqueIdx) : MUniqueIdx(InUniqueIdx) { }
+			Chaos::FUniqueIdx UniqueIdx() const { return MUniqueIdx; }
+			Chaos::FUniqueIdx MUniqueIdx;
+		};
+		const FMockHandle MockHandle0(HandleExternal0.UniqueIdx());
+
+		// Unregister particle 0. The GT delegate should fire synchronously
+		// inside UnregisterObject, cleaning the GT cache immediately.
+		Solver->UnregisterObject(Proxy0);
+
+		// GT cache for the unregistered particle should already be gone —
+		// no solver advance needed.
+		REQUIRE(TestUserData->GetData_GT(MockHandle0) == nullptr);
+
+		// The other particle's GT data should be unaffected
+		REQUIRE(*TestUserData->GetData_GT(HandleExternal1) == TestString2);
+
+		// Let the solver process the destruction so PT is also cleaned up
+		AdvanceAndWait(Solver, DeltaTime);
+		AdvanceAndWait(Solver, DeltaTime);
+	}
+
+	SECTION("Recycled unique index does not return stale GT data")
+	{
+		// Set data on particle 0
+		TestUserData->SetData_GT(HandleExternal0, TestString1);
+		REQUIRE(*TestUserData->GetData_GT(HandleExternal0) == TestString1);
+
+		// Save the unique index, then destroy the particle
+		const Chaos::FUniqueIdx OldUniqueIdx = HandleExternal0.UniqueIdx();
+		Solver->UnregisterObject(Proxy0);
+		AdvanceAndWait(Solver, DeltaTime);
+		AdvanceAndWait(Solver, DeltaTime);
+
+		// Create a new particle — it may or may not reuse the same
+		// FUniqueIdx, but regardless, querying the old index should not
+		// return stale data.
+		struct FMockHandle
+		{
+			FMockHandle(Chaos::FUniqueIdx InUniqueIdx) : MUniqueIdx(InUniqueIdx) { }
+			Chaos::FUniqueIdx UniqueIdx() const { return MUniqueIdx; }
+			Chaos::FUniqueIdx MUniqueIdx;
+		};
+		const FMockHandle OldHandle(OldUniqueIdx);
+		REQUIRE(TestUserData->GetData_GT(OldHandle) == nullptr);
+
+		// If we create a brand-new particle and set fresh data on it,
+		// GetData_GT must return the new data, not anything from the old particle.
+		Chaos::FSingleParticlePhysicsProxy* Proxy2 = Chaos::FSingleParticlePhysicsProxy::Create(Chaos::FGeometryParticle::CreateParticle());
+		Chaos::FRigidBodyHandle_External& HandleExternal2 = Proxy2->GetGameThreadAPI();
+		HandleExternal2.SetGeometry(BoxGeom);
+		Solver->RegisterObject(Proxy2);
+		AdvanceAndWait(Solver, DeltaTime);
+		AdvanceAndWait(Solver, DeltaTime);
+
+		TestUserData->SetData_GT(HandleExternal2, TestString2);
+		REQUIRE(*TestUserData->GetData_GT(HandleExternal2) == TestString2);
+
+		// And the old index still shows nothing (unless it was reused by
+		// the new particle, in which case it correctly shows new data).
+		const FMockHandle NewHandle(HandleExternal2.UniqueIdx());
+		if (OldUniqueIdx == HandleExternal2.UniqueIdx())
+		{
+			// Index was reused — must see the new data, not the old
+			REQUIRE(*TestUserData->GetData_GT(OldHandle) == TestString2);
+		}
+		else
+		{
+			// Index was not reused — old slot must still be empty
+			REQUIRE(TestUserData->GetData_GT(OldHandle) == nullptr);
+		}
+
+		Solver->UnregisterObject(Proxy2);
+		AdvanceAndWait(Solver, DeltaTime);
+		AdvanceAndWait(Solver, DeltaTime);
+	}
+
+	Solver->WaitOnPendingTasks_External();
+	Module->DestroySolver(Solver);
+}
diff --git a/Engine/Source/Runtime/Experimental/Chaos/Private/PBDRigidsSolver.cpp b/Engine/Source/Runtime/Experimental/Chaos/Private/PBDRigidsSolver.cpp
--- a/Engine/Source/Runtime/Experimental/Chaos/Private/PBDRigidsSolver.cpp
+++ b/Engine/Source/Runtime/Experimental/Chaos/Private/PBDRigidsSolver.cpp
@@ -912,6 +912,9 @@ namespace Chaos
 		// Remove the proxy from the GT proxy map
 		FUniqueIdx UniqueIdx = Proxy->GetGameThreadAPI().UniqueIdx();
 
+		// Notify GT listeners before the proxy is enqueued for destruction
+		EventParticleUnregisteredGT.Broadcast(Proxy, UniqueIdx);
+
 		{
 			// The collision manager is also modified from FPBDRigidsSolver::PushPhysicsState which uses the 
 			// MarshallingManager lock so we need to lock this here as well.
@@ -965,6 +968,16 @@ namespace Chaos
 		});
 	}
 
+	FDelegateHandle FPBDRigidsSolver::AddParticleUnregisteredGTCallback(FSolverParticleUnregisteredGT::FDelegate InDelegate)
+	{
+		return EventParticleUnregisteredGT.Add(InDelegate);
+	}
+
+	bool FPBDRigidsSolver::RemoveParticleUnregisteredGTCallback(FDelegateHandle InHandle)
+	{
+		return EventParticleUnregisteredGT.Remove(InHandle);
+	}
+
 	void FPBDRigidsSolver::RegisterObject(FGeometryCollectionPhysicsProxy* InProxy)
 	{
 		LLM_SCOPE(ELLMTag::ChaosBody);
diff --git a/Engine/Source/Runtime/Experimental/Chaos/Public/Chaos/Framework/PhysicsSolverBase.h b/Engine/Source/Runtime/Experimental/Chaos/Public/Chaos/Framework/PhysicsSolverBase.h
--- a/Engine/Source/Runtime/Experimental/Chaos/Public/Chaos/Framework/PhysicsSolverBase.h
+++ b/Engine/Source/Runtime/Experimental/Chaos/Public/Chaos/Framework/PhysicsSolverBase.h
@@ -422,6 +422,7 @@ namespace Chaos
 			}
 
 			RegisterSimCallbackObject_External(NewCallbackObject);
+			NewCallbackObject->PostInitialize_External();
 			return NewCallbackObject;
 		}
 
@@ -451,6 +452,7 @@ namespace Chaos
 
 		void UnregisterAndFreeSimCallbackObject_External(ISimCallbackObject* SimCallbackObject)
 		{
+			SimCallbackObject->PreShutdown_External();
 			MarshallingManager.UnregisterSimCallbackObject_External(SimCallbackObject);
 
 			// Unregister external callback
diff --git a/Engine/Source/Runtime/Experimental/Chaos/Public/Chaos/SimCallbackObject.h b/Engine/Source/Runtime/Experimental/Chaos/Public/Chaos/SimCallbackObject.h
--- a/Engine/Source/Runtime/Experimental/Chaos/Public/Chaos/SimCallbackObject.h
+++ b/Engine/Source/Runtime/Experimental/Chaos/Public/Chaos/SimCallbackObject.h
@@ -98,6 +98,16 @@ public:
 		OnPostInitialize_Internal();
 	}
 
+	void PostInitialize_External()
+	{
+		OnPostInitialize_External();
+	}
+
+	void PreShutdown_External()
+	{
+		OnPreShutdown_External();
+	}
+
 	virtual void ProcessInputs_Internal(int32 PhysicsStep)
 	{
 		OnProcessInputs_Internal(PhysicsStep);
@@ -285,6 +295,20 @@ private:
 	*/
 	virtual void OnPostInitialize_Internal() { }
 
+	/**
+	* Called on the external thread after the callback object is registered with the solver.
+	* The solver is guaranteed to be available via GetSolver() at this point.
+	* Use this for GT-side setup such as registering with solver delegates.
+	*/
+	virtual void OnPostInitialize_External() { }
+
+	/**
+	* Called on the external thread before the callback object is unregistered from the solver.
+	* Symmetric to OnPostInitialize_External — use this to clean up any solver-level registrations
+	* (e.g. delegates) that were set up during initialization.
+	*/
+	virtual void OnPreShutdown_External() { }
+
 	/** Called before any sim callbacks are triggered but after async inputs has marshaled over
 	* Register with ESimCallbackOptions::ProcessInputsInternal to get callback */
 	virtual void OnProcessInputs_Internal(int32 PhysicsStep)
diff --git a/Engine/Source/Runtime/Experimental/Chaos/Public/PBDRigidsSolver.h b/Engine/Source/Runtime/Experimental/Chaos/Public/PBDRigidsSolver.h
--- a/Engine/Source/Runtime/Experimental/Chaos/Public/PBDRigidsSolver.h
+++ b/Engine/Source/Runtime/Experimental/Chaos/Public/PBDRigidsSolver.h
@@ -41,6 +41,8 @@ class FPerSolverFieldSystem;
 
 extern int32 ChaosSolverParticlePoolNumFrameUntilShrink;
 
+DECLARE_MULTICAST_DELEGATE_TwoParams(FSolverParticleUnregisteredGT, Chaos::FSingleParticlePhysicsProxy* /*Proxy*/, Chaos::FUniqueIdx /*UniqueIdx*/);
+
 namespace ChaosTest
 {
 	template <typename TSolver>
@@ -145,6 +147,11 @@ namespace Chaos
 		CHAOS_API void RegisterObject(Chaos::FCharacterGroundConstraint* GTConstraint);
 		CHAOS_API void UnregisterObject(Chaos::FCharacterGroundConstraint* GTConstraint);
 
+		/** Fired on the game thread when a single particle physics proxy is unregistered.
+		    Called during UnregisterObject before the proxy is enqueued for destruction. */
+		UE_INTERNAL CHAOS_API FDelegateHandle AddParticleUnregisteredGTCallback(FSolverParticleUnregisteredGT::FDelegate InDelegate);
+		UE_INTERNAL CHAOS_API bool            RemoveParticleUnregisteredGTCallback(FDelegateHandle InHandle);
+
 		//
 		//  Simulation API
 		//
@@ -441,6 +448,8 @@ namespace Chaos
 
 		FSolverSerializer Serializer;
 
+		FSolverParticleUnregisteredGT EventParticleUnregisteredGT;
+
 		struct FPendingDestroyInfo
 		{
 			FSingleParticlePhysicsProxy* Proxy;
